Live Link KiPro Device Plugin - Pure C++ Implementation

Implemented Live Link Device plugin for AJA Ki Pro recording devices, porting functionality from Switchboard to Live Link Hub framework.

Key Changes:
- Pure C++ implementation using Unreal HTTP module and JSON parsing
- FTSTicker-based timing for reconnection and polling
- REST API client for Ki Pro remote control (record/stop, slate/take metadata)
- Periodic reconnection attempts (5s interval when disconnected)
- Recording session delegate subscription for real-time slate/take synchronization
- Settings change triggers immediate connection attempt
- Firmware version decoding (32-bit encoded integer to dotted version string)
- Mock Ki Pro device simulator for testing without hardware (Python, Source/Scripts/)

Implementation:
- Consolidated single ticker at 0.1s for transport polling and reconnection
- Time accumulator pattern for 5-second reconnection interval
- Async HTTP operations with lambda callbacks
- Transport state polling with timeout handling
- Connection/Recording capabilities via Live Link Device framework

Testing: Mock device simulator validated against Switchboard implementation

#Synced-CL 51338485

diff --git a/Engine/Plugins/VirtualProduction/LiveLinkKiProDevice/LiveLinkKiProDevice.uplugin b/Engine/Plugins/VirtualProduction/LiveLinkKiProDevice/LiveLinkKiProDevice.uplugin
new file mode 100644
--- /dev/null
+++ b/Engine/Plugins/VirtualProduction/LiveLinkKiProDevice/LiveLinkKiProDevice.uplugin
@@ -0,0 +1,35 @@
+{
+	"FileVersion": 3,
+	"Version": 1,
+	"VersionName": "1.0",
+	"FriendlyName": "Live Link KiPro Device",
+	"Description": "AJA Ki Pro recorder support for Live Link Hub with recording and connection capabilities",
+	"Category": "Virtual Production",
+	"CreatedBy": "Epic Games, Inc.",
+	"CreatedByURL": "https://epicgames.com",
+	"CanContainContent": true,
+	"IsBetaVersion": false,
+	"IsExperimentalVersion": true,
+	"Installed": false,
+	"Modules": [
+		{
+			"Name": "LiveLinkKiProDevice",
+			"Type": "Editor",
+			"LoadingPhase": "Default"
+		}
+	],
+	"Plugins": [
+		{
+			"Name": "LiveLinkDevice",
+			"Enabled": true
+		},
+		{
+			"Name": "LiveLinkGenericRecordingDevice",
+			"Enabled": true
+		},
+		{
+			"Name": "PythonScriptPlugin",
+			"Enabled": true
+		}
+	]
+}
diff --git a/Engine/Plugins/VirtualProduction/LiveLinkKiProDevice/Source/LiveLinkKiProDevice/LiveLinkKiProDevice.Build.cs b/Engine/Plugins/VirtualProduction/LiveLinkKiProDevice/Source/LiveLinkKiProDevice/LiveLinkKiProDevice.Build.cs
new file mode 100644
--- /dev/null
+++ b/Engine/Plugins/VirtualProduction/LiveLinkKiProDevice/Source/LiveLinkKiProDevice/LiveLinkKiProDevice.Build.cs
@@ -0,0 +1,35 @@
+// Copyright Epic Games, Inc. All Rights Reserved.
+
+using UnrealBuildTool;
+
+public class LiveLinkKiProDevice : ModuleRules
+{
+	public LiveLinkKiProDevice(ReadOnlyTargetRules Target) : base(Target)
+	{
+		PCHUsage = ModuleRules.PCHUsageMode.UseExplicitOrSharedPCHs;
+
+		PublicDependencyModuleNames.AddRange(
+			new string[]
+			{
+				"Core",
+			}
+		);
+
+		PrivateDependencyModuleNames.AddRange(
+			new string[]
+			{
+				"CoreUObject",
+				"Engine",
+				"Slate",
+				"SlateCore",
+				"LiveLinkDevice",
+				"LiveLinkInterface",
+				"Projects",
+				"HTTP",
+				"Json",
+				"JsonUtilities",
+				"Networking"
+			}
+		);
+	}
+}
diff --git a/Engine/Plugins/VirtualProduction/LiveLinkKiProDevice/Source/LiveLinkKiProDevice/Private/LiveLinkKiProDeviceBase.cpp b/Engine/Plugins/VirtualProduction/LiveLinkKiProDevice/Source/LiveLinkKiProDevice/Private/LiveLinkKiProDeviceBase.cpp
new file mode 100644
--- /dev/null
+++ b/Engine/Plugins/VirtualProduction/LiveLinkKiProDevice/Source/LiveLinkKiProDevice/Private/LiveLinkKiProDeviceBase.cpp
@@ -0,0 +1,578 @@
+// Copyright Epic Games, Inc. All Rights Reserved.
+
+#include "LiveLinkKiProDeviceBase.h"
+#include "HttpModule.h"
+#include "Interfaces/IHttpResponse.h"
+
+#include "Dom/JsonValue.h"
+#include "Dom/JsonObject.h"
+#include "Runtime/Networking/Public/Interfaces/IPv4/IPv4Endpoint.h"
+#include "Serialization/JsonTypes.h"
+#include "Serialization/JsonReader.h"
+#include "Serialization/JsonWriter.h"
+#include "Serialization/JsonSerializer.h"
+
+#include "LiveLinkHub/ILiveLinkRecordingSession.h"
+#include "TimerManager.h"
+#include "Engine/World.h"
+#include "Engine/Engine.h"
+
+#include "Interfaces/IPv4/IPv4Endpoint.h"
+#include "GenericPlatform/GenericPlatformHttp.h"
+
+#define LOCTEXT_NAMESPACE "LiveLinkKiProDevice"
+
+//~ Begin ULiveLinkDevice interface
+
+TSubclassOf<ULiveLinkDeviceSettings> ULiveLinkKiProDeviceBase::GetSettingsClass() const
+{
+	return SettingsClass ? SettingsClass : TSubclassOf<ULiveLinkDeviceSettings>(ULiveLinkKiProDeviceSettings::StaticClass());
+}
+
+FText ULiveLinkKiProDeviceBase::GetDisplayName() const
+{
+	const ULiveLinkKiProDeviceSettings* DeviceSettings = GetDeviceSettings<ULiveLinkKiProDeviceSettings>();
+	return FText::FromString(FString::Printf(TEXT("Ki Pro (%s)"), *DeviceSettings->IpAddress));
+}
+
+EDeviceHealth ULiveLinkKiProDeviceBase::GetDeviceHealth() const
+{
+	switch (ConnectionStatus)
+	{
+	case ELiveLinkDeviceConnectionStatus::Connected:
+		return bIsRecording ? EDeviceHealth::Good : EDeviceHealth::Nominal;
+	case ELiveLinkDeviceConnectionStatus::Connecting:
+		return EDeviceHealth::Info;
+	case ELiveLinkDeviceConnectionStatus::Disconnecting:
+		return EDeviceHealth::Warning;
+	case ELiveLinkDeviceConnectionStatus::Disconnected:
+	default:
+		return EDeviceHealth::Error;
+	}
+}
+
+FText ULiveLinkKiProDeviceBase::GetHealthText() const
+{
+	switch (ConnectionStatus)
+	{
+	case ELiveLinkDeviceConnectionStatus::Connected:
+		return bIsRecording ? LOCTEXT("Recording", "Recording") : LOCTEXT("Connected", "Connected - Ready");
+	case ELiveLinkDeviceConnectionStatus::Connecting:
+		return LOCTEXT("Connecting", "Connecting...");
+	case ELiveLinkDeviceConnectionStatus::Disconnecting:
+		return LOCTEXT("Disconnecting", "Disconnecting...");
+	case ELiveLinkDeviceConnectionStatus::Disconnected:
+	default:
+		return LOCTEXT("Disconnected", "Disconnected");
+	}
+}
+
+void ULiveLinkKiProDeviceBase::OnDeviceAdded()
+{
+	// Subscribe to recording session delegates for slate/take changes
+	ILiveLinkRecordingSession& Session = ILiveLinkRecordingSession::Get();
+	Session.OnSlateNameChanged().AddUObject(this, &ULiveLinkKiProDeviceBase::HandleSlateNameChanged);
+	Session.OnTakeNumberChanged().AddUObject(this, &ULiveLinkKiProDeviceBase::HandleTakeNumberChanged);
+
+	// Start main device ticker (handles reconnection and polling)
+	// Use high frequency (0.1s) for transport polling, accumulate time for reconnection
+	DeviceTickerHandle =
+		FTSTicker::GetCoreTicker().AddTicker(
+			FTickerDelegate::CreateUObject(this, &ULiveLinkKiProDeviceBase::DeviceTick),
+			PollingIntervalSeconds
+		);
+}
+
+void ULiveLinkKiProDeviceBase::OnDeviceRemoved()
+{
+	// Unsubscribe from recording session delegates
+	ILiveLinkRecordingSession& Session = ILiveLinkRecordingSession::Get();
+	Session.OnSlateNameChanged().RemoveAll(this);
+	Session.OnTakeNumberChanged().RemoveAll(this);
+
+	// Stop device ticker
+	if (DeviceTickerHandle.IsValid())
+	{
+		FTSTicker::GetCoreTicker().RemoveTicker(DeviceTickerHandle);
+		DeviceTickerHandle.Reset();
+	}
+
+	if (ConnectionStatus != ELiveLinkDeviceConnectionStatus::Disconnected)
+	{
+		Disconnect_Implementation();
+	}
+}
+
+
+void ULiveLinkKiProDeviceBase::OnSettingChanged(const FPropertyChangedEvent& InPropertyChangedEvent)
+{
+	// When settings change, attempt immediate reconnection if disconnected
+	// This provides instant feedback when user updates IP address or port
+	if (ConnectionStatus == ELiveLinkDeviceConnectionStatus::Disconnected)
+	{
+		UE_LOGF(LogTemp, Log, "Settings changed, attempting connection to Ki Pro...");
+		AttemptReconnection();
+	}
+}
+
+ELiveLinkDeviceConnectionStatus ULiveLinkKiProDeviceBase::GetConnectionStatus_Implementation() const
+{
+	return ConnectionStatus;
+}
+
+FString ULiveLinkKiProDeviceBase::GetHardwareId_Implementation() const
+{
+	const ULiveLinkKiProDeviceSettings* DeviceSettings = GetDeviceSettings<ULiveLinkKiProDeviceSettings>();
+	return FString::Printf(TEXT("%s:%d"), *DeviceSettings->IpAddress, DeviceSettings->Port);
+}
+
+bool ULiveLinkKiProDeviceBase::SetHardwareId_Implementation(const FString& HardwareID)
+{
+	ULiveLinkKiProDeviceSettings* DeviceSettings = GetDeviceSettings<ULiveLinkKiProDeviceSettings>();
+
+	FIPv4Endpoint OutEndpoint;
+	if (FIPv4Endpoint::Parse(HardwareID, OutEndpoint))
+	{
+		DeviceSettings->IpAddress = OutEndpoint.Address.ToString();
+		DeviceSettings->Port = OutEndpoint.Port;
+		return true;
+	}
+
+	return false;
+}
+
+bool ULiveLinkKiProDeviceBase::CanSetHardwareId_Implementation()
+{
+	// Can set hardware ID when disconnected
+	return ConnectionStatus == ELiveLinkDeviceConnectionStatus::Disconnected;
+}
+
+bool ULiveLinkKiProDeviceBase::Connect_Implementation()
+{
+	if (ConnectionStatus != ELiveLinkDeviceConnectionStatus::Disconnected)
+	{
+		return false;
+	}
+
+	ConnectionStatus = ELiveLinkDeviceConnectionStatus::Connecting;
+
+	// Test connection by getting firmware version
+	GetParameter(TEXT("eParamID_SWVersion"), [this](bool bSuccess, const FString& Value, const FString& Text)
+	{
+		if (bSuccess)
+		{
+			ULiveLinkKiProDeviceSettings* DeviceSettings = GetDeviceSettings<ULiveLinkKiProDeviceSettings>();
+
+			// Text contains the encoded version as a string (e.g., "67239937")
+			// Convert to int and decode to human-readable format (e.g., "4.2.0.1")
+			int32 EncodedVersion = FCString::Atoi(*Text);
+			DeviceSettings->FirmwareVersion = DecodeFirmwareVersion(EncodedVersion);
+
+			ConnectionStatus = ELiveLinkDeviceConnectionStatus::Connected;
+			UE_LOG(LogTemp, Log, TEXT("Connected to Ki Pro at %s, Firmware: %s"), *DeviceSettings->IpAddress, *DeviceSettings->FirmwareVersion);
+		}
+		else
+		{
+			ConnectionStatus = ELiveLinkDeviceConnectionStatus::Disconnected;
+			UE_LOGF(LogTemp, Error, "Failed to connect to Ki Pro - device unresponsive");
+		}
+	});
+
+	return true;
+}
+
+bool ULiveLinkKiProDeviceBase::Disconnect_Implementation()
+{
+	ConnectionStatus = ELiveLinkDeviceConnectionStatus::Disconnecting;
+
+	// Stop any polling (handled by main ticker now)
+	CurrentPollingOperation = EPollingOperation::None;
+	bIsRecording = false;
+	ConnectionStatus = ELiveLinkDeviceConnectionStatus::Disconnected;
+	UE_LOGF(LogTemp, Log, "Disconnected from Ki Pro");
+
+	return true;
+}
+
+bool ULiveLinkKiProDeviceBase::StartRecording_Implementation()
+{
+	if (ConnectionStatus != ELiveLinkDeviceConnectionStatus::Connected)
+	{
+		UE_LOGF(LogTemp, Warning, "Cannot start recording - not connected");
+		return false;
+	}
+
+	// Get slate/take from recording session
+	ILiveLinkRecordingSession& RecordingSession = ILiveLinkRecordingSession::Get();
+
+	const FString Slate = RecordingSession.GetSlateName();
+	const int32 Take = RecordingSession.GetTakeNumber();
+
+	UE_LOG(LogTemp, Log, TEXT("Starting recording: Slate=%s, Take=%d"), *Slate, Take);
+
+	// Set media state for recording
+	SetMediaStateForRecordPlay();
+
+	// Set slate/take metadata
+	SetParameter(TEXT("eParamID_CustomClipName"), Slate, [](bool bSuccess) {});
+	SetParameter(TEXT("eParamID_CustomTake"), FString::FromInt(Take), [](bool bSuccess) {});
+
+	// Send record command
+	SendTransportCommand(TEXT("Record Command"));
+
+	// Begin polling for state change (handled by main ticker)
+	CurrentPollingOperation = EPollingOperation::WaitingForRecording;
+	PollingStartTime = FPlatformTime::Seconds();
+
+	return true;
+}
+
+bool ULiveLinkKiProDeviceBase::StopRecording_Implementation()
+{
+	if (ConnectionStatus != ELiveLinkDeviceConnectionStatus::Connected)
+	{
+		UE_LOGF(LogTemp, Warning, "Cannot stop recording - not connected");
+		return false;
+	}
+
+	UE_LOGF(LogTemp, Log, "Stopping recording");
+
+	// Set media state for recording
+	SetMediaStateForRecordPlay();
+
+	// Send stop command
+	SendTransportCommand(TEXT("Stop Command"));
+
+	// Begin polling for state change to Idle (handled by main ticker)
+	CurrentPollingOperation = EPollingOperation::WaitingForIdle;
+	PollingStartTime = FPlatformTime::Seconds();
+
+	return true;
+}
+
+bool ULiveLinkKiProDeviceBase::IsRecording_Implementation() const
+{
+	return bIsRecording;
+}
+
+//~ End ILiveLinkDeviceCapability_Recording interface
+
+//~ Begin Private Methods
+
+FString ULiveLinkKiProDeviceBase::GetBaseUrl() const
+{
+	const ULiveLinkKiProDeviceSettings* DeviceSettings = GetDeviceSettings<ULiveLinkKiProDeviceSettings>();
+	return FString::Printf(TEXT("http://%s:%d"), *DeviceSettings->IpAddress, DeviceSettings->Port);
+}
+
+FString ULiveLinkKiProDeviceBase::DecodeFirmwareVersion(int32 VersionBits)
+{
+	// Decode 32-bit integer to dotted version string
+	// Format: (major<<24) | (minor<<16) | (patch<<8) | build
+	// Example: 67239937 (0x04020001) -> "4.2.0.1"
+	const uint8 Major = (VersionBits & 0xFF000000) >> 24;
+	const uint8 Minor = (VersionBits & 0x00FF0000) >> 16;
+	const uint8 Patch = (VersionBits & 0x0000FF00) >> 8;
+	const uint8 Build = (VersionBits & 0x000000FF);
+
+	return FString::Printf(TEXT("%d.%d.%d.%d"), Major, Minor, Patch, Build);
+}
+
+void ULiveLinkKiProDeviceBase::GetParameter(const FString& ParamId, TFunction<void(bool bSuccess, const FString& Value, const FString& Text)> OnComplete)
+{
+	// KiPro devices are usually on the local network and 2.0s should be a reasonable default timeout.
+	const float HttpRequestTimeout = 2.0f;
+	TSharedRef<IHttpRequest, ESPMode::ThreadSafe> HttpRequest = FHttpModule::Get().CreateRequest();
+	HttpRequest->SetVerb(TEXT("GET"));
+	HttpRequest->SetURL(GetBaseUrl() + TEXT("/options?") + ParamId);
+	HttpRequest->SetTimeout(HttpRequestTimeout);
+
+	// Capture weak pointer to protect against device destruction before callback
+	TWeakObjectPtr<ULiveLinkKiProDeviceBase> WeakThis(this);
+	HttpRequest->OnProcessRequestComplete().BindWeakLambda(this, [WeakThis, OnComplete](FHttpRequestPtr Request, FHttpResponsePtr Response, bool bWasSuccessful)
+	{
+		// Check if device still exists before executing callback
+		if (WeakThis.IsValid())
+		{
+			HandleGetParameterResponse(Request, Response, bWasSuccessful, OnComplete);
+		}
+	});
+
+	HttpRequest->ProcessRequest();
+}
+
+void ULiveLinkKiProDeviceBase::SetParameter(const FString& ParamId, const FString& Value, TFunction<void(bool bSuccess)> OnComplete)
+{
+	// KiPro devices are usually on the local network and 2.0s should be a reasonable default timeout.
+	const float HttpRequestTimeout = 2.0f;
+	TSharedRef<IHttpRequest, ESPMode::ThreadSafe> HttpRequest = FHttpModule::Get().CreateRequest();
+	HttpRequest->SetVerb(TEXT("GET"));
+
+	FString EncodedValue = FGenericPlatformHttp::UrlEncode(Value);
+	FString URL = FString::Printf(TEXT("%s/config?action=set&paramid=%s&value=%s"), *GetBaseUrl(), *ParamId, *EncodedValue);
+	HttpRequest->SetURL(URL);
+	HttpRequest->SetTimeout(HttpRequestTimeout);
+
+	// Capture weak pointer to protect against device destruction before callback
+	TWeakObjectPtr<ULiveLinkKiProDeviceBase> WeakThis(this);
+	HttpRequest->OnProcessRequestComplete().BindWeakLambda(this, [WeakThis, OnComplete](FHttpRequestPtr Request, FHttpResponsePtr Response, bool bWasSuccessful)
+	{
+		// Check if device still exists before executing callback
+		if (WeakThis.IsValid())
+		{
+			HandleSetParameterResponse(Request, Response, bWasSuccessful, OnComplete);
+		}
+	});
+
+	HttpRequest->ProcessRequest();
+}
+
+void ULiveLinkKiProDeviceBase::GetTransportState(TFunction<void(bool bSuccess, const FString& StateText)> OnComplete)
+{
+	GetParameter(TEXT("eParamID_TransportState"), [OnComplete](bool bSuccess, const FString& Value, const FString& Text)
+	{
+		OnComplete(bSuccess, Text);
+	});
+}
+
+void ULiveLinkKiProDeviceBase::SendTransportCommand(const FString& CommandDescription)
+{
+	// For simplicity, map command descriptions to known values
+	// In full implementation, would query eParamID_TransportCommand options first
+	int32 CommandValue = 0;
+
+	if (CommandDescription == TEXT("Record Command"))
+	{
+		CommandValue = 4;
+	}
+	else if (CommandDescription == TEXT("Stop Command"))
+	{
+		CommandValue = 1;
+	}
+	else if (CommandDescription == TEXT("Play Command"))
+	{
+		CommandValue = 3;
+	}
+
+	SetParameter(TEXT("eParamID_TransportCommand"), FString::FromInt(CommandValue), [](bool bSuccess) {});
+}
+
+void ULiveLinkKiProDeviceBase::SetMediaStateForRecordPlay()
+{
+	// Check if in Data-LAN mode and switch to Record-Play (value 0)
+	GetParameter(TEXT("eParamID_MediaState"), [this](bool bSuccess, const FString& Value, const FString& Text)
+	{
+		if (bSuccess && Text == TEXT("Data - LAN"))
+		{
+			UE_LOGF(LogTemp, Log, "Switching Ki Pro from Data-LAN to Record-Play mode");
+			SetParameter(TEXT("eParamID_MediaState"), TEXT("0"), [](bool bSuccess) {});
+		}
+	});
+}
+
+void ULiveLinkKiProDeviceBase::PollTransportState()
+{
+	// Prevent request flooding on slow connections - don't send new request if one is in-flight
+	if (bPollingRequestInFlight)
+	{
+		return;
+	}
+
+	// Check for timeout
+	if (FPlatformTime::Seconds() - PollingStartTime > PollingTimeoutSeconds)
+	{
+		UE_LOGF(LogTemp, Warning, "Timeout waiting for transport state change");
+		CurrentPollingOperation = EPollingOperation::None;
+		bPollingRequestInFlight = false;
+		return;
+	}
+
+	// Mark request as in-flight before sending
+	bPollingRequestInFlight = true;
+
+	// Poll transport state
+	GetTransportState([this](bool bSuccess, const FString& StateText)
+	{
+		OnTransportStateReceived(bSuccess, StateText);
+	});
+}
+
+void ULiveLinkKiProDeviceBase::OnTransportStateReceived(bool bSuccess, const FString& StateText)
+{
+	// Clear in-flight flag now that request has completed
+	bPollingRequestInFlight = false;
+
+	if (!bSuccess)
+	{
+		UE_LOGF(LogTemp, Error, "Failed to get transport state");
+		CurrentPollingOperation = EPollingOperation::None;
+		return;
+	}
+
+	bool bOperationComplete = false;
+
+	switch (CurrentPollingOperation)
+	{
+	case EPollingOperation::WaitingForRecording:
+		if (StateText == TEXT("Recording"))
+		{
+			bIsRecording = true;
+			bOperationComplete = true;
+			UE_LOGF(LogTemp, Log, "Recording started on Ki Pro");
+		}
+		break;
+
+	case EPollingOperation::WaitingForIdle:
+		if (StateText == TEXT("Idle"))
+		{
+			bIsRecording = false;
+			bOperationComplete = true;
+			OnRecordingStopped();
+		}
+		break;
+
+	case EPollingOperation::WaitingForPlay:
+		if (StateText == TEXT("Playing Forward"))
+		{
+			bOperationComplete = true;
+			UE_LOGF(LogTemp, Log, "Playback started on Ki Pro");
+		}
+		break;
+
+	default:
+		break;
+	}
+
+	if (bOperationComplete)
+	{
+		CurrentPollingOperation = EPollingOperation::None;
+	}
+}
+
+void ULiveLinkKiProDeviceBase::OnRecordingStopped()
+{
+	// Get clip name for logging
+	GetParameter(TEXT("eParamID_CurrentClip"), [this](bool bSuccess, const FString& Value, const FString& Text)
+	{
+		if (bSuccess)
+		{
+			UE_LOG(LogTemp, Log, TEXT("Recording stopped. Clip: %s"), *Text);
+		}
+
+		// Auto-play if enabled
+		const ULiveLinkKiProDeviceSettings* DeviceSettings = GetDeviceSettings<ULiveLinkKiProDeviceSettings>();
+		if (DeviceSettings->bAutoPlayAfterStop)
+		{
+			UE_LOGF(LogTemp, Log, "Auto-play enabled, starting playback");
+			SendTransportCommand(TEXT("Play Command"));
+
+			// Continue polling for playback state (handled by main ticker)
+			CurrentPollingOperation = EPollingOperation::WaitingForPlay;
+			PollingStartTime = FPlatformTime::Seconds();
+		}
+	});
+}
+
+void ULiveLinkKiProDeviceBase::HandleGetParameterResponse(FHttpRequestPtr Request, FHttpResponsePtr Response, bool bWasSuccessful, TFunction<void(bool, const FString&, const FString&)> OnComplete)
+{
+	if (!bWasSuccessful || !Response.IsValid() || Response->GetResponseCode() != 200)
+	{
+		OnComplete(false, FString(), FString());
+		return;
+	}
+
+	FString ResponseContent = Response->GetContentAsString();
+
+	// Parse JSON response - AJA returns array directly, not wrapped in object
+	TArray<TSharedPtr<FJsonValue>> OptionsArray;
+	TSharedRef<TJsonReader<>> JsonReader = TJsonReaderFactory<>::Create(ResponseContent);
+
+	if (FJsonSerializer::Deserialize(JsonReader, OptionsArray))
+	{
+		// Find selected option (AJA uses "true"/"false" strings, not booleans)
+		for (const TSharedPtr<FJsonValue>& OptionValue : OptionsArray)
+		{
+			TSharedPtr<FJsonObject> OptionObject = OptionValue->AsObject();
+			if (OptionObject.IsValid())
+			{
+				FString SelectedStr;
+				OptionObject->TryGetStringField(TEXT("selected"), SelectedStr);
+
+				if (SelectedStr == TEXT("true"))
+				{
+					FString Value, Text;
+					OptionObject->TryGetStringField(TEXT("value"), Value);
+					OptionObject->TryGetStringField(TEXT("text"), Text);
+					OnComplete(true, Value, Text);
+					return;
+				}
+			}
+		}
+	}
+
+	OnComplete(false, FString(), FString());
+}
+
+void ULiveLinkKiProDeviceBase::HandleSetParameterResponse(FHttpRequestPtr Request, FHttpResponsePtr Response, bool bWasSuccessful, TFunction<void(bool)> OnComplete)
+{
+	OnComplete(bWasSuccessful && Response.IsValid() && Response->GetResponseCode() == 200);
+}
+
+bool ULiveLinkKiProDeviceBase::DeviceTick(float DeltaTime)
+{
+	// Handle periodic reconnection attempts
+	if (ConnectionStatus == ELiveLinkDeviceConnectionStatus::Disconnected)
+	{
+		ReconnectionTimeAccumulator += DeltaTime;
+		if (ReconnectionTimeAccumulator >= ReconnectionIntervalSeconds)
+		{
+			ReconnectionTimeAccumulator = 0.0f;
+			AttemptReconnection();
+		}
+	}
+	else
+	{
+		// Reset accumulator when connected
+		ReconnectionTimeAccumulator = 0.0f;
+	}
+
+	// Handle transport state polling
+	if (CurrentPollingOperation != EPollingOperation::None)
+	{
+		PollTransportState();
+	}
+
+	// Return true to continue ticking
+	return true;
+}
+
+void ULiveLinkKiProDeviceBase::AttemptReconnection()
+{
+	UE_LOGF(LogTemp, Log, "Attempting to reconnect to Ki Pro...");
+	ILiveLinkDeviceCapability_Connection::Execute_Connect(this);
+}
+
+void ULiveLinkKiProDeviceBase::HandleSlateNameChanged(FStringView InSlateName)
+{
+	// Update slate name on device if connected
+	if (ConnectionStatus == ELiveLinkDeviceConnectionStatus::Connected)
+	{
+		FString SlateName(InSlateName);
+		UE_LOG(LogTemp, Log, TEXT("Slate name changed to: %s"), *SlateName);
+		SetParameter(TEXT("eParamID_CustomClipName"), SlateName, [](bool bSuccess) {});
+	}
+}
+
+void ULiveLinkKiProDeviceBase::HandleTakeNumberChanged(int32 InTakeNumber)
+{
+	// Update take number on device if connected
+	if (ConnectionStatus == ELiveLinkDeviceConnectionStatus::Connected)
+	{
+		UE_LOGF(LogTemp, Log, "Take number changed to: %d", InTakeNumber);
+		SetParameter(TEXT("eParamID_CustomTake"), FString::FromInt(InTakeNumber), [](bool bSuccess) {});
+	}
+}
+
+//~ End Private Methods
+
+#undef LOCTEXT_NAMESPACE
diff --git a/Engine/Plugins/VirtualProduction/LiveLinkKiProDevice/Source/LiveLinkKiProDevice/Private/LiveLinkKiProDeviceModule.cpp b/Engine/Plugins/VirtualProduction/LiveLinkKiProDevice/Source/LiveLinkKiProDevice/Private/LiveLinkKiProDeviceModule.cpp
new file mode 100644
--- /dev/null
+++ b/Engine/Plugins/VirtualProduction/LiveLinkKiProDevice/Source/LiveLinkKiProDevice/Private/LiveLinkKiProDeviceModule.cpp
@@ -0,0 +1,5 @@
+// Copyright Epic Games, Inc. All Rights Reserved.
+
+#include "Modules/ModuleManager.h"
+
+IMPLEMENT_MODULE(FDefaultModuleImpl, LiveLinkKiProDevice);
diff --git a/Engine/Plugins/VirtualProduction/LiveLinkKiProDevice/Source/LiveLinkKiProDevice/Public/LiveLinkKiProDeviceBase.h b/Engine/Plugins/VirtualProduction/LiveLinkKiProDevice/Source/LiveLinkKiProDevice/Public/LiveLinkKiProDeviceBase.h
new file mode 100644
--- /dev/null
+++ b/Engine/Plugins/VirtualProduction/LiveLinkKiProDevice/Source/LiveLinkKiProDevice/Public/LiveLinkKiProDeviceBase.h
@@ -0,0 +1,144 @@
+// Copyright Epic Games, Inc. All Rights Reserved.
+
+#pragma once
+
+#include "CoreMinimal.h"
+#include "LiveLinkDevice.h"
+#include "LiveLinkDeviceCapability_Connection.h"
+#include "LiveLinkDeviceCapability_Recording.h"
+#include "LiveLinkKiProDeviceSettings.h"
+#include "Interfaces/IHttpRequest.h"
+#include "Containers/Ticker.h"
+
+#include "LiveLinkKiProDeviceBase.generated.h"
+
+/**
+ * Concrete implementation for AJA Ki Pro recording devices
+ * Provides both Connection and Recording capabilities via REST API
+ */
+UCLASS(BlueprintType, meta = (DisplayName = "KiPro Device", ToolTip = "Support for AJA KiPro"))
+class LIVELINKKIPRODEVICE_API ULiveLinkKiProDeviceBase : public ULiveLinkDevice
+	, public ILiveLinkDeviceCapability_Connection
+	, public ILiveLinkDeviceCapability_Recording
+{
+	GENERATED_BODY()
+
+public:
+	//~ Begin ULiveLinkDevice interface
+	virtual TSubclassOf<ULiveLinkDeviceSettings> GetSettingsClass() const override;
+	virtual FText GetDisplayName() const override;
+	virtual EDeviceHealth GetDeviceHealth() const override;
+	virtual FText GetHealthText() const override;
+	virtual void OnDeviceAdded() override;
+	virtual void OnDeviceRemoved() override;
+	virtual void OnSettingChanged(const FPropertyChangedEvent& InPropertyChangedEvent) override;
+	//~ End ULiveLinkDevice interface
+
+	//~ Begin ILiveLinkDeviceCapability_Connection interface
+	virtual ELiveLinkDeviceConnectionStatus GetConnectionStatus_Implementation() const override;
+	virtual FString GetHardwareId_Implementation() const override;
+	virtual bool SetHardwareId_Implementation(const FString& HardwareID) override;
+	virtual bool CanSetHardwareId_Implementation() override;
+	virtual bool Connect_Implementation() override;
+	virtual bool Disconnect_Implementation() override;
+	//~ End ILiveLinkDeviceCapability_Connection interface
+
+	//~ Begin ILiveLinkDeviceCapability_Recording interface
+	virtual bool StartRecording_Implementation() override;
+	virtual bool StopRecording_Implementation() override;
+	virtual bool IsRecording_Implementation() const override;
+	//~ End ILiveLinkDeviceCapability_Recording interface
+
+private:
+	/** Build URL for Ki Pro REST API */
+	FString GetBaseUrl() const;
+
+	/** Decode firmware version from 32-bit integer to dotted string (e.g., 67239937 -> "4.2.0.1") */
+	static FString DecodeFirmwareVersion(int32 VersionBits);
+
+	/** Get a parameter from the Ki Pro device */
+	void GetParameter(const FString& ParamId, TFunction<void(bool bSuccess, const FString& Value, const FString& Text)> OnComplete);
+
+	/** Set a parameter on the Ki Pro device */
+	void SetParameter(const FString& ParamId, const FString& Value, TFunction<void(bool bSuccess)> OnComplete);
+
+	/** Get the transport state */
+	void GetTransportState(TFunction<void(bool bSuccess, const FString& StateText)> OnComplete);
+
+	/** Send a transport command by finding its value from description */
+	void SendTransportCommand(const FString& CommandDescription);
+
+	/** Ensure device is in Record-Play mode (not Data-LAN) */
+	void SetMediaStateForRecordPlay();
+
+	/** Poll transport state until target state is reached */
+	void PollTransportState();
+
+	/** Called when transport state polling completes */
+	void OnTransportStateReceived(bool bSuccess, const FString& StateText);
+
+	/** Called when recording stop is confirmed */
+	void OnRecordingStopped();
+
+	/** Handle HTTP request completion for GetParameter */
+	static void HandleGetParameterResponse(FHttpRequestPtr Request, FHttpResponsePtr Response, bool bWasSuccessful, TFunction<void(bool, const FString&, const FString&)> OnComplete);
+
+	/** Handle HTTP request completion for SetParameter */
+	static void HandleSetParameterResponse(FHttpRequestPtr Request, FHttpResponsePtr Response, bool bWasSuccessful, TFunction<void(bool)> OnComplete);
+
+	/** Main device tick function (handles reconnection and polling) */
+	bool DeviceTick(float DeltaTime);
+
+	/** Attempt reconnection if disconnected */
+	void AttemptReconnection();
+
+	/** Handle slate name changes from recording session */
+	void HandleSlateNameChanged(FStringView InSlateName);
+
+	/** Handle take number changes from recording session */
+	void HandleTakeNumberChanged(int32 InTakeNumber);
+
+protected:
+	/** Settings class to use for Ki Pro devices */
+	UPROPERTY(EditDefaultsOnly, Category="Live Link Device")
+	TSubclassOf<ULiveLinkDeviceSettings> SettingsClass;
+
+	/** Current connection status */
+	UPROPERTY()
+	ELiveLinkDeviceConnectionStatus ConnectionStatus = ELiveLinkDeviceConnectionStatus::Disconnected;
+
+	/** Whether the device is currently recording */
+	UPROPERTY()
+	bool bIsRecording = false;
+
+	/** Ticker handle for main device tick (reconnection and polling) */
+	FTSTicker::FDelegateHandle DeviceTickerHandle;
+
+	/** What operation we're polling for */
+	enum class EPollingOperation
+	{
+		None,
+		WaitingForRecording,
+		WaitingForIdle,
+		WaitingForPlay
+	};
+	EPollingOperation CurrentPollingOperation = EPollingOperation::None;
+
+	/** Time when polling started (for timeout) */
+	double PollingStartTime = 0.0;
+
+	/** Time accumulator for reconnection attempts */
+	float ReconnectionTimeAccumulator = 0.0f;
+
+	/** Flag to prevent polling request flood on slow connections */
+	bool bPollingRequestInFlight = false;
+
+	/** Polling timeout in seconds */
+	static constexpr double PollingTimeoutSeconds = 4.0;
+
+	/** Polling interval in seconds */
+	static constexpr float PollingIntervalSeconds = 0.1f;
+
+	/** Reconnection attempt interval in seconds */
+	static constexpr float ReconnectionIntervalSeconds = 5.0f;
+};
diff --git a/Engine/Plugins/VirtualProduction/LiveLinkKiProDevice/Source/LiveLinkKiProDevice/Public/LiveLinkKiProDeviceSettings.h b/Engine/Plugins/VirtualProduction/LiveLinkKiProDevice/Source/LiveLinkKiProDevice/Public/LiveLinkKiProDeviceSettings.h
new file mode 100644
--- /dev/null
+++ b/Engine/Plugins/VirtualProduction/LiveLinkKiProDevice/Source/LiveLinkKiProDevice/Public/LiveLinkKiProDeviceSettings.h
@@ -0,0 +1,33 @@
+// Copyright Epic Games, Inc. All Rights Reserved.
+
+#pragma once
+
+#include "CoreMinimal.h"
+#include "LiveLinkDevice.h"
+#include "LiveLinkKiProDeviceSettings.generated.h"
+
+/**
+ * Settings for AJA Ki Pro recording device
+ */
+UCLASS()
+class LIVELINKKIPRODEVICE_API ULiveLinkKiProDeviceSettings : public ULiveLinkDeviceSettings
+{
+	GENERATED_BODY()
+
+public:
+	/** IP address of the Ki Pro device */
+	UPROPERTY(EditAnywhere, Category="Ki Pro", meta=(DisplayName="IP Address"))
+	FString IpAddress = TEXT("192.168.1.100");
+
+	/** HTTP port for REST API communication */
+	UPROPERTY(EditAnywhere, Category="Ki Pro", AdvancedDisplay)
+	uint16 Port = 80;
+
+	/** Automatically start playback after recording stops */
+	UPROPERTY(EditAnywhere, Category="Ki Pro", meta=(DisplayName="Auto Play After Stop Recording"))
+	bool bAutoPlayAfterStop = false;
+
+	/** Firmware version retrieved from device (read-only) */
+	UPROPERTY(VisibleAnywhere, Category="Ki Pro", AdvancedDisplay, meta=(DisplayName="Firmware Version"))
+	FString FirmwareVersion;
+};
diff --git a/Engine/Plugins/VirtualProduction/LiveLinkKiProDevice/Source/Scripts/mock_kipro_README.md b/Engine/Plugins/VirtualProduction/LiveLinkKiProDevice/Source/Scripts/mock_kipro_README.md
new file mode 100644
--- /dev/null
+++ b/Engine/Plugins/VirtualProduction/LiveLinkKiProDevice/Source/Scripts/mock_kipro_README.md
@@ -0,0 +1,217 @@
+# Mock Ki Pro Device Simulator
+
+A Python-based simulator that implements the AJA Ki Pro REST API for testing the Live Link Ki Pro plugin without physical hardware.
+
+## Features
+
+- Implements key Ki Pro REST API endpoints
+- Simulates transport state machine (Idle → Recording → Idle)
+- Tracks recording parameters (slate, take, clip name)
+- Media state management (Record-Play vs Data-LAN mode)
+- Real-time command logging
+- Configurable IP and port
+
+## Requirements
+
+- Python 3.6 or higher
+- No external dependencies (uses Python standard library)
+
+## Usage
+
+### Basic Usage (Port 80)
+
+Run as administrator/root since port 80 requires elevated privileges:
+
+```bash
+# Windows (run as Administrator)
+python mock_kipro_device.py
+
+# Linux/Mac (run with sudo)
+sudo python3 mock_kipro_device.py
+```
+
+### Custom Port (No Admin Required)
+
+```bash
+python mock_kipro_device.py --port 8080
+```
+
+### Custom IP Address
+
+```bash
+python mock_kipro_device.py --ip 192.168.1.100 --port 80
+```
+
+### All Options
+
+```bash
+python mock_kipro_device.py --help
+```
+
+## Testing with Live Link Hub
+
+1. **Start the simulator:**
+   ```bash
+   python mock_kipro_device.py --port 8080
+   ```
+
+2. **Launch Live Link Hub**
+
+3. **Add Ki Pro Device:**
+   - Click "Add Device"
+   - Select "Ki Pro Device"
+   - Device appears in device list
+
+4. **Configure IP Address:**
+   - Edit device settings
+   - Set IP Address: `127.0.0.1` (for local testing)
+   - Set Port: `8080` (or 80 if using default)
+
+5. **Connect:**
+   - Click "Connect" button in device row
+   - Simulator will log: `[GET] eParamID_SWVersion -> 4.2.0.1`
+   - Device status should show "Connected - Ready"
+
+6. **Test Recording:**
+   - Set slate and take in Live Link Hub recording session
+   - Click "Start Recording"
+   - Simulator logs:
+     ```
+     [SET] eParamID_CustomClipName = MySlate
+     [SET] eParamID_CustomTake = 1
+     [TRANSPORT COMMAND] Record Command (value=4)
+       -> Recording started: MySlate_Take1
+     ```
+   - Click "Stop Recording"
+   - Simulator logs:
+     ```
+     [TRANSPORT COMMAND] Stop Command (value=1)
+       -> Stopped (was Recording)
+     ```
+
+## Simulator Output Example
+
+```
+============================================================
+Mock Ki Pro Device Simulator
+============================================================
+Listening on: http://0.0.0.0:8080
+Firmware Version: 4.2.0.1
+
+Waiting for connections from Live Link Hub...
+Press Ctrl+C to stop
+============================================================
+
+[GET] eParamID_SWVersion -> 4.2.0.1
+[MEDIA STATE] Changed to: Record - Play
+[SLATE] Set to: Scene_1_Shot_A
+[TAKE] Set to: 3
+
+[TRANSPORT COMMAND] Record Command (value=4)
+  -> Recording started: Scene_1_Shot_A_Take3
+
+[GET] eParamID_TransportState -> Recording
+
+[TRANSPORT COMMAND] Stop Command (value=1)
+  -> Stopped (was Recording)
+
+[GET] eParamID_CurrentClip -> Scene_1_Shot_A_Take3
+```
+
+## Implemented API Endpoints
+
+### GET /options?<param_id>
+
+Returns parameter information with current value and options.
+
+**Example:** `GET /options?eParamID_SWVersion`
+
+**Response:**
+```json
+{
+  "param_id": "eParamID_SWVersion",
+  "options": [
+    {
+      "value": "4020001",
+      "text": "4.2.0.1",
+      "selected": true
+    }
+  ]
+}
+```
+
+### GET /config?action=set&paramid=<id>&value=<val>
+
+Sets a parameter value.
+
+**Example:** `GET /config?action=set&paramid=eParamID_CustomClipName&value=MySlate`
+
+**Response:**
+```json
+{
+  "status": "ok"
+}
+```
+
+## Supported Parameters
+
+- `eParamID_SWVersion` - Firmware version (read-only)
+- `eParamID_TransportState` - Current transport state (Idle, Recording, Playing Forward)
+- `eParamID_MediaState` - Media state (Record-Play, Data-LAN)
+- `eParamID_TransportCommand` - Transport commands (Stop=1, Play=3, Record=4)
+- `eParamID_CustomClipName` - Slate name
+- `eParamID_CustomTake` - Take number
+- `eParamID_CurrentClip` - Current clip name (generated from slate+take)
+
+## Transport State Machine
+
+```
+Idle → Record Command → Recording
+Recording → Stop Command → Idle
+Idle → Play Command → Playing Forward
+Playing Forward → Stop Command → Idle
+```
+
+## Troubleshooting
+
+### Port 80 Permission Denied
+
+On Linux/Mac, port 80 requires root:
+```bash
+sudo python3 mock_kipro_device.py
+```
+
+Or use a higher port:
+```bash
+python3 mock_kipro_device.py --port 8080
+```
+
+### Connection Refused
+
+- Check firewall settings
+- Verify IP address and port match in Live Link Hub settings
+- Ensure simulator is running before attempting connection
+
+### Device Not Appearing in Live Link Hub
+
+- Rebuild the plugin after C++ changes
+- Verify the plugin is enabled in Live Link Hub
+- Check that `ULiveLinkDeviceKiProBase` is not marked as `Abstract`
+
+## Known Limitations
+
+- Does not implement file download/Data-LAN mode functionality
+- Simplified transport command mapping (hard-coded values)
+- No authentication
+- Single-threaded (one request at a time)
+- No persistent storage of clips
+
+## Development Notes
+
+The simulator tracks state in memory and resets on restart. It's designed for basic functional testing of the Live Link Ki Pro plugin's connection, recording, and transport control features.
+
+For more advanced testing scenarios, you can modify the Python script to add:
+- Delayed state transitions (simulate real hardware timing)
+- Error injection (simulate connection failures)
+- Additional parameters
+- File system integration
diff --git a/Engine/Plugins/VirtualProduction/LiveLinkKiProDevice/Source/Scripts/mock_kipro_device.py b/Engine/Plugins/VirtualProduction/LiveLinkKiProDevice/Source/Scripts/mock_kipro_device.py
new file mode 100644
--- /dev/null
+++ b/Engine/Plugins/VirtualProduction/LiveLinkKiProDevice/Source/Scripts/mock_kipro_device.py
@@ -0,0 +1,367 @@
+#!/usr/bin/env python3
+"""
+Mock AJA Ki Pro Device Simulator
+
+Simulates the REST API of an AJA Ki Pro device for testing the Live Link Ki Pro plugin.
+Implements the key endpoints and parameter management needed for recording control.
+
+Usage:
+    python mock_kipro_device.py [--port PORT] [--ip IP]
+
+Example:
+    python mock_kipro_device.py --port 80 --ip 192.168.1.100
+"""
+
+import argparse
+import json
+import time
+from http.server import HTTPServer, BaseHTTPRequestHandler
+from urllib.parse import urlparse, parse_qs
+from threading import Lock
+
+
+class KiProState:
+    """Manages the state of the mock Ki Pro device"""
+
+    def __init__(self):
+        self.lock = Lock()
+
+        # Device parameters
+        self.firmware_version = "4.2.0.1"
+        self.transport_state = "Idle"  # Idle, Recording, Playing Forward, etc.
+        self.media_state = "Record - Play"  # "Record - Play" or "Data - LAN"
+        self.custom_clip_name = "DefaultClip"
+        self.custom_take = "1"
+        self.current_clip = ""
+
+        # Transport state machine
+        self.state_transitions = {
+            "Idle": ["Recording", "Playing Forward"],
+            "Recording": ["Idle"],
+            "Playing Forward": ["Idle"],
+        }
+
+        # Parameter definitions (param_id -> {value, text, options})
+        # NOTE: AJA uses "true"/"false" STRINGS for selected, not booleans
+        # NOTE: For version params, "text" must be encoded as: (major<<24)|(minor<<16)|(patch<<8)|build
+        #       Version 4.2.0.1 = (4<<24)|(2<<16)|(0<<8)|1 = 67239937 = 0x04020001
+        self.parameters = {
+            "eParamID_SWVersion": {
+                "value": "67239937",
+                "text": "67239937",  # Encoded: 4.2.0.1 -> 0x04020001 -> 67239937
+                "options": [{"value": "67239937", "text": "67239937", "selected": "true"}]
+            },
+            "eParamID_TransportState": {
+                "value": "0",
+                "text": "Idle",
+                "options": [
+                    {"value": "0", "text": "Idle", "selected": "true"},
+                    {"value": "1", "text": "Recording", "selected": "false"},
+                    {"value": "2", "text": "Playing Forward", "selected": "false"},
+                ]
+            },
+            "eParamID_MediaState": {
+                "value": "0",
+                "text": "Record - Play",
+                "options": [
+                    {"value": "0", "text": "Record - Play", "selected": "true"},
+                    {"value": "1", "text": "Data - LAN", "selected": "false"},
+                ]
+            },
+            "eParamID_TransportCommand": {
+                "value": "0",
+                "text": "No Command",
+                "options": [
+                    {"value": "0", "text": "No Command", "selected": "true"},
+                    {"value": "1", "text": "Stop Command", "selected": "false"},
+                    {"value": "2", "text": "Play Reverse Command", "selected": "false"},
+                    {"value": "3", "text": "Play Command", "selected": "false"},
+                    {"value": "4", "text": "Record Command", "selected": "false"},
+                ]
+            },
+            "eParamID_CustomClipName": {
+                "value": "DefaultClip",
+                "text": "DefaultClip",
+                "options": [{"value": "DefaultClip", "text": "DefaultClip", "selected": "true"}]
+            },
+            "eParamID_CustomTake": {
+                "value": "1",
+                "text": "1",
+                "options": [{"value": "1", "text": "1", "selected": "true"}]
+            },
+            "eParamID_CurrentClip": {
+                "value": "",
+                "text": "",
+                "options": [{"value": "", "text": "", "selected": "true"}]
+            },
+        }
+
+    def get_parameter(self, param_id):
+        """Get parameter value and options"""
+        with self.lock:
+            if param_id in self.parameters:
+                return self.parameters[param_id]
+            return None
+
+    def set_parameter(self, param_id, value):
+        """Set parameter value"""
+        with self.lock:
+            if param_id == "eParamID_TransportCommand":
+                return self._handle_transport_command(value)
+            elif param_id == "eParamID_MediaState":
+                return self._set_media_state(value)
+            elif param_id == "eParamID_CustomClipName":
+                return self._set_custom_clip_name(value)
+            elif param_id == "eParamID_CustomTake":
+                return self._set_custom_take(value)
+            else:
+                print(f"Warning: Unhandled set parameter: {param_id} = {value}")
+                return True
+
+    def _handle_transport_command(self, value):
+        """Handle transport command"""
+        command_map = {
+            "1": "Stop Command",
+            "3": "Play Command",
+            "4": "Record Command",
+        }
+
+        command = command_map.get(value, "Unknown")
+        print(f"\n[TRANSPORT COMMAND] {command} (value={value})")
+
+        if value == "4":  # Record
+            if self.transport_state == "Idle":
+                self.transport_state = "Recording"
+                self._update_transport_state_param()
+                # Generate clip name from slate and take
+                clip_name = f"{self.custom_clip_name}_Take{self.custom_take}"
+                self.current_clip = clip_name
+                self._update_current_clip_param(clip_name)
+                print(f"  -> Recording started: {clip_name}")
+            else:
+                print(f"  -> Cannot record from state: {self.transport_state}")
+
+        elif value == "1":  # Stop
+            if self.transport_state in ["Recording", "Playing Forward"]:
+                prev_state = self.transport_state
+                self.transport_state = "Idle"
+                self._update_transport_state_param()
+                print(f"  -> Stopped (was {prev_state})")
+            else:
+                print(f"  -> Already in state: {self.transport_state}")
+
+        elif value == "3":  # Play
+            if self.transport_state == "Idle":
+                self.transport_state = "Playing Forward"
+                self._update_transport_state_param()
+                print(f"  -> Playback started")
+            else:
+                print(f"  -> Cannot play from state: {self.transport_state}")
+
+        return True
+
+    def _set_media_state(self, value):
+        """Set media state (Record-Play vs Data-LAN)"""
+        if value == "0":
+            self.media_state = "Record - Play"
+            text = "Record - Play"
+        elif value == "1":
+            self.media_state = "Data - LAN"
+            text = "Data - LAN"
+        else:
+            return False
+
+        print(f"[MEDIA STATE] Changed to: {text}")
+
+        # Update parameter (use string "true"/"false")
+        param = self.parameters["eParamID_MediaState"]
+        param["value"] = value
+        param["text"] = text
+        for opt in param["options"]:
+            opt["selected"] = "true" if opt["value"] == value else "false"
+
+        return True
+
+    def _set_custom_clip_name(self, value):
+        """Set custom clip name (slate)"""
+        self.custom_clip_name = value
+        print(f"[SLATE] Set to: {value}")
+
+        # Update parameter
+        param = self.parameters["eParamID_CustomClipName"]
+        param["value"] = value
+        param["text"] = value
+        param["options"] = [{"value": value, "text": value, "selected": "true"}]
+
+        return True
+
+    def _set_custom_take(self, value):
+        """Set custom take number"""
+        self.custom_take = value
+        print(f"[TAKE] Set to: {value}")
+
+        # Update parameter
+        param = self.parameters["eParamID_CustomTake"]
+        param["value"] = value
+        param["text"] = value
+        param["options"] = [{"value": value, "text": value, "selected": "true"}]
+
+        return True
+
+    def _update_transport_state_param(self):
+        """Update transport state parameter"""
+        state_map = {
+            "Idle": "0",
+            "Recording": "1",
+            "Playing Forward": "2",
+        }
+
+        value = state_map.get(self.transport_state, "0")
+        param = self.parameters["eParamID_TransportState"]
+        param["value"] = value
+        param["text"] = self.transport_state
+
+        # Update selected option (use string "true"/"false")
+        for opt in param["options"]:
+            opt["selected"] = "true" if opt["text"] == self.transport_state else "false"
+
+    def _update_current_clip_param(self, clip_name):
+        """Update current clip parameter"""
+        param = self.parameters["eParamID_CurrentClip"]
+        param["value"] = clip_name
+        param["text"] = clip_name
+        param["options"] = [{"value": clip_name, "text": clip_name, "selected": "true"}]
+
+
+class KiProRequestHandler(BaseHTTPRequestHandler):
+    """HTTP request handler for Ki Pro REST API"""
+
+    # Class-level state shared across requests
+    device_state = None
+
+    def log_message(self, format, *args):
+        """Override to customize logging"""
+        # Log requests to stdout
+        pass  # Suppress default logging, we'll print our own
+
+    def do_GET(self):
+        """Handle GET requests"""
+        parsed_url = urlparse(self.path)
+        query_params = parse_qs(parsed_url.query)
+
+        if parsed_url.path == '/options':
+            self._handle_get_options(query_params)
+        elif parsed_url.path == '/config':
+            self._handle_set_config(query_params)
+        elif parsed_url.path == '/descriptors':
+            self._handle_get_descriptors(query_params)
+        else:
+            self.send_error(404, "Not Found")
+
+    def _handle_get_options(self, query_params):
+        """Handle /options?<param_id> requests"""
+        # Extract param_id from query string
+        # The query is just the param_id itself, e.g., "eParamID_SWVersion"
+        param_id = self.path.split('?')[1] if '?' in self.path else None
+
+        if not param_id:
+            self.send_error(400, "Missing parameter ID")
+            return
+
+        param_data = self.device_state.get_parameter(param_id)
+
+        if param_data:
+            print(f"[GET] {param_id} -> {param_data['text']}")
+
+            # AJA REST API returns array directly, NOT wrapped in object
+            response = param_data["options"]
+
+            self._send_json_response(response)
+        else:
+            print(f"[GET] Unknown parameter: {param_id}")
+            self.send_error(404, f"Parameter not found: {param_id}")
+
+    def _handle_set_config(self, query_params):
+        """Handle /config?action=set&paramid=<id>&value=<val> requests"""
+        # Parse query manually since it's not standard form encoding
+        parts = self.path.split('?')[1] if '?' in self.path else ""
+        params = {}
+        for part in parts.split('&'):
+            if '=' in part:
+                key, val = part.split('=', 1)
+                params[key] = val
+
+        action = params.get('action')
+        param_id = params.get('paramid')
+        value = params.get('value')
+
+        if action != 'set' or not param_id or value is None:
+            self.send_error(400, "Invalid config request")
+            return
+
+        # URL decode value
+        import urllib.parse
+        value = urllib.parse.unquote(value)
+
+        print(f"[SET] {param_id} = {value}")
+
+        success = self.device_state.set_parameter(param_id, value)
+
+        if success:
+            # Return success response
+            response = {"status": "ok"}
+            self._send_json_response(response)
+        else:
+            self.send_error(400, "Failed to set parameter")
+
+    def _handle_get_descriptors(self, query_params):
+        """Handle /descriptors?paramid=* requests (used for caching)"""
+        # The client requests all descriptors with paramid=*
+        # For simplicity, return empty array - client will fall back to individual queries
+        print("[DESCRIPTORS] Requested (returning empty - client will query individually)")
+        response = []
+        self._send_json_response(response)
+
+    def _send_json_response(self, data):
+        """Send JSON response"""
+        json_data = json.dumps(data, indent=2)
+        self.send_response(200)
+        self.send_header('Content-Type', 'application/json')
+        self.send_header('Content-Length', len(json_data))
+        self.end_headers()
+        self.wfile.write(json_data.encode('utf-8'))
+
+
+def main():
+    parser = argparse.ArgumentParser(description='Mock AJA Ki Pro Device Simulator')
+    parser.add_argument('--port', type=int, default=80, help='Port to listen on (default: 80)')
+    parser.add_argument('--ip', type=str, default='0.0.0.0', help='IP address to bind to (default: 0.0.0.0)')
+    args = parser.parse_args()
+
+    # Initialize device state
+    KiProRequestHandler.device_state = KiProState()
+
+    # Create HTTP server
+    server_address = (args.ip, args.port)
+    httpd = HTTPServer(server_address, KiProRequestHandler)
+
+    print("=" * 60)
+    print(f"Mock Ki Pro Device Simulator")
+    print("=" * 60)
+    print(f"Listening on: http://{args.ip}:{args.port}")
+    print(f"Firmware Version: {KiProRequestHandler.device_state.firmware_version}")
+    print()
+    print("Waiting for connections from Live Link Hub...")
+    print("Press Ctrl+C to stop")
+    print("=" * 60)
+    print()
+
+    try:
+        httpd.serve_forever()
+    except KeyboardInterrupt:
+        print("\n\nShutting down server...")
+        httpd.shutdown()
+
+
+if __name__ == "__main__":
+    main()
diff --git a/Engine/Source/Programs/LiveLinkHub/Source/LiveLinkHub.Target.cs b/Engine/Source/Programs/LiveLinkHub/Source/LiveLinkHub.Target.cs
--- a/Engine/Source/Programs/LiveLinkHub/Source/LiveLinkHub.Target.cs
+++ b/Engine/Source/Programs/LiveLinkHub/Source/LiveLinkHub.Target.cs
@@ -57,6 +57,7 @@ public class LiveLinkHubTarget : TargetRules
 
 			"LiveLinkDevice",
 			"LiveLinkUnrealDevice",
+			"LiveLinkKiProDevice",
 
 			"LiveLinkOpenVR",
 
